Delegate for tools base classes we want to have project-specific telemetry for.

This was needed so we could track tool usage utilizing project specific analytics classes, allowing projects to opt-in with the OnEditorToolStarted delegate instead of always sending an engine level analytic event.

diff --git a/Engine/Plugins/Editor/DataValidation/Source/DataValidation/DataValidation.Build.cs b/Engine/Plugins/Editor/DataValidation/Source/DataValidation/DataValidation.Build.cs
--- a/Engine/Plugins/Editor/DataValidation/Source/DataValidation/DataValidation.Build.cs
+++ b/Engine/Plugins/Editor/DataValidation/Source/DataValidation/DataValidation.Build.cs
@@ -46,7 +46,8 @@ public class DataValidation : ModuleRules
 				"TypedElementRuntime",
 				"SourceControl",
 				"UncontrolledChangelists",
-				"RHI"
+				"RHI",
+				"EditorToolEvents"
 			}
 		);
 	}
diff --git a/Engine/Plugins/Editor/DataValidation/Source/DataValidation/Private/EditorValidatorBase.cpp b/Engine/Plugins/Editor/DataValidation/Source/DataValidation/Private/EditorValidatorBase.cpp
--- a/Engine/Plugins/Editor/DataValidation/Source/DataValidation/Private/EditorValidatorBase.cpp
+++ b/Engine/Plugins/Editor/DataValidation/Source/DataValidation/Private/EditorValidatorBase.cpp
@@ -4,6 +4,7 @@
 
 #include "AssetRegistry/AssetData.h"
 #include "AssetRegistry/AssetDataToken.h"
+#include "EditorToolDelegates.h"
 #include "EditorValidatorSubsystem.h"
 #include "Misc/DataValidation.h"
 #include "Trace/Trace.inl"
@@ -79,6 +80,11 @@ PRAGMA_ENABLE_DEPRECATION_WARNINGS
 	}
 	Result = CombineDataValidationResults(Result, ExtractValidationState(Context)); // Extract messages and validation state from members (primarily for use by BP interface functions)
 
+	if (bTriedToValidate)
+	{
+		FEditorToolDelegates::OnEditorToolStarted.Broadcast(GetClass()->GetName());
+	}
+
 	return Result;
 }
 
diff --git a/Engine/Source/Editor/Blutility/Blutility.Build.cs b/Engine/Source/Editor/Blutility/Blutility.Build.cs
--- a/Engine/Source/Editor/Blutility/Blutility.Build.cs
+++ b/Engine/Source/Editor/Blutility/Blutility.Build.cs
@@ -47,6 +47,7 @@ public class Blutility : ModuleRules
 				"ImageWrapper",
 				"ImageWriteQueue",
 				"DeveloperSettings",
+				"EditorToolEvents"
 			}
 			);
 	}
diff --git a/Engine/Source/Editor/Blutility/Classes/EditorUtilityObject.h b/Engine/Source/Editor/Blutility/Classes/EditorUtilityObject.h
--- a/Engine/Source/Editor/Blutility/Classes/EditorUtilityObject.h
+++ b/Engine/Source/Editor/Blutility/Classes/EditorUtilityObject.h
@@ -23,6 +23,9 @@ class BLUTILITY_API UEditorUtilityObject : public UObject, public IAssetRegistry
 	UFUNCTION(BlueprintCallable, BlueprintImplementableEvent, Category = "Editor")
 	void Run();
 
+	// we use this to call tool started delegate
+	virtual void PostLoad() override;
+
 	//~ Begin IAssetRegistryTagProviderInterface interface
 	virtual bool ShouldAddCDOTagsToBlueprintClass() const override 
 	{ 
diff --git a/Engine/Source/Editor/Blutility/Classes/EditorUtilityWidget.h b/Engine/Source/Editor/Blutility/Classes/EditorUtilityWidget.h
--- a/Engine/Source/Editor/Blutility/Classes/EditorUtilityWidget.h
+++ b/Engine/Source/Editor/Blutility/Classes/EditorUtilityWidget.h
@@ -63,6 +63,10 @@ public:
 	virtual bool IsEditorUtility() const override { return true; }
 
 protected:
+
+	/** Used to send event OnEditorToolStarted event on initialize */
+	virtual void NativeOnInitialized() override;
+
 	/** The display name for tabs spawned with this widget */
 	UPROPERTY(Category = Config, EditDefaultsOnly, BlueprintReadWrite, AssetRegistrySearchable)
 	FText TabDisplayName;
diff --git a/Engine/Source/Editor/Blutility/Private/EditorUtilityObject.cpp b/Engine/Source/Editor/Blutility/Private/EditorUtilityObject.cpp
--- a/Engine/Source/Editor/Blutility/Private/EditorUtilityObject.cpp
+++ b/Engine/Source/Editor/Blutility/Private/EditorUtilityObject.cpp
@@ -2,9 +2,17 @@
 
 #include "EditorUtilityObject.h"
 
+#include "EditorToolDelegates.h"
 /////////////////////////////////////////////////////
 
 UEditorUtilityObject::UEditorUtilityObject(const FObjectInitializer& ObjectInitializer)
 	: Super(ObjectInitializer)
 {
 }
+
+void UEditorUtilityObject::PostLoad()
+{
+	Super::PostLoad();
+
+	FEditorToolDelegates::OnEditorToolStarted.Broadcast(GetClass()->GetName());
+}
diff --git a/Engine/Source/Editor/Blutility/Private/EditorUtilityWidget.cpp b/Engine/Source/Editor/Blutility/Private/EditorUtilityWidget.cpp
--- a/Engine/Source/Editor/Blutility/Private/EditorUtilityWidget.cpp
+++ b/Engine/Source/Editor/Blutility/Private/EditorUtilityWidget.cpp
@@ -3,6 +3,7 @@
 #include "EditorUtilityWidget.h"
 
 #include "Blueprint/WidgetTree.h"
+#include "EditorToolDelegates.h"
 #include "Internationalization/Internationalization.h"
 #include "Misc/AssertionMacros.h"
 #include "ScopedTransaction.h"
@@ -31,6 +32,12 @@ void UEditorUtilityWidget::ExecuteDefaultAction()
 	Run();
 }
 
+void UEditorUtilityWidget::NativeOnInitialized()
+{
+	Super::NativeOnInitialized();
+
+	FEditorToolDelegates::OnEditorToolStarted.Broadcast(GetClass()->GetName());
+}
 
 
 #undef LOCTEXT_NAMESPACE
\ No newline at end of file
diff --git a/Engine/Source/Editor/EditorToolEvents/EditorToolEvents.Build.cs b/Engine/Source/Editor/EditorToolEvents/EditorToolEvents.Build.cs
new file mode 100644
--- /dev/null
+++ b/Engine/Source/Editor/EditorToolEvents/EditorToolEvents.Build.cs
@@ -0,0 +1,49 @@
+// Copyright Epic Games, Inc. All Rights Reserved.
+
+using UnrealBuildTool;
+
+public class EditorToolEvents : ModuleRules
+{
+	public EditorToolEvents(ReadOnlyTargetRules Target) : base(Target)
+	{
+		PCHUsage = ModuleRules.PCHUsageMode.UseExplicitOrSharedPCHs;
+		
+		PublicIncludePaths.AddRange(
+			new string[] {
+				// ... add public include paths required here ...
+			}
+			);
+				
+		
+		PrivateIncludePaths.AddRange(
+			new string[] {
+				// ... add other private include paths required here ...
+			}
+			);
+			
+		
+		PublicDependencyModuleNames.AddRange(
+			new string[]
+			{
+				"Core",
+				"Engine"
+				// ... add other public dependencies that you statically link with here ...
+			}
+			);
+			
+		
+		PrivateDependencyModuleNames.AddRange(
+			new string[]
+			{
+			}
+			);
+		
+		
+		DynamicallyLoadedModuleNames.AddRange(
+			new string[]
+			{
+				// ... add any modules that your module loads dynamically here ...
+			}
+			);
+	}
+}
diff --git a/Engine/Source/Editor/EditorToolEvents/Private/EditorToolDelegates.cpp b/Engine/Source/Editor/EditorToolEvents/Private/EditorToolDelegates.cpp
new file mode 100644
--- /dev/null
+++ b/Engine/Source/Editor/EditorToolEvents/Private/EditorToolDelegates.cpp
@@ -0,0 +1,5 @@
+// Copyright Epic Games, Inc. All Rights Reserved.
+
+#include "EditorToolDelegates.h"
+
+FEditorToolDelegates::FEditorToolStarted	FEditorToolDelegates::OnEditorToolStarted;
diff --git a/Engine/Source/Editor/EditorToolEvents/Private/EditorToolEventsModule.cpp b/Engine/Source/Editor/EditorToolEvents/Private/EditorToolEventsModule.cpp
new file mode 100644
--- /dev/null
+++ b/Engine/Source/Editor/EditorToolEvents/Private/EditorToolEventsModule.cpp
@@ -0,0 +1,15 @@
+// Copyright Epic Games, Inc. All Rights Reserved.
+
+#include "EditorToolEventsModule.h"
+
+
+void FEditorToolEventsModule::StartupModule()
+{
+
+}
+
+void FEditorToolEventsModule::ShutdownModule()
+{
+}
+
+IMPLEMENT_MODULE(FEditorToolEventsModule, EditorToolEvents)
\ No newline at end of file
diff --git a/Engine/Source/Editor/EditorToolEvents/Public/EditorToolDelegates.h b/Engine/Source/Editor/EditorToolEvents/Public/EditorToolDelegates.h
new file mode 100644
--- /dev/null
+++ b/Engine/Source/Editor/EditorToolEvents/Public/EditorToolDelegates.h
@@ -0,0 +1,19 @@
+// Copyright Epic Games, Inc. All Rights Reserved.
+
+/**
+ * Struct of static delegates for editor tools. Delegates intended to be used by projects to send tool telemetry events
+ */
+
+#pragma once
+#include "Delegates/Delegate.h"
+
+struct FEditorToolDelegates
+{
+
+public:
+	/**Broadcasted for editor tools OnLoad or OnInitialize, if available to notify when tool used */
+	DECLARE_MULTICAST_DELEGATE_OneParam(FEditorToolStarted, const FString &/*class name*/);
+	static EDITORTOOLEVENTS_API FEditorToolStarted OnEditorToolStarted;
+	
+	/*More delegates can be added, starting minimal so only events we use are added*/
+};
diff --git a/Engine/Source/Editor/EditorToolEvents/Public/EditorToolEventsModule.h b/Engine/Source/Editor/EditorToolEvents/Public/EditorToolEventsModule.h
new file mode 100644
--- /dev/null
+++ b/Engine/Source/Editor/EditorToolEvents/Public/EditorToolEventsModule.h
@@ -0,0 +1,14 @@
+// Copyright Epic Games, Inc. All Rights Reserved.
+
+#pragma once
+
+#include "Modules/ModuleManager.h"
+
+class FEditorToolEventsModule : public IModuleInterface
+{
+public:
+
+	/** IModuleInterface implementation */
+	virtual void StartupModule() override;
+	virtual void ShutdownModule() override;
+};
diff --git a/Engine/Source/Runtime/InteractiveToolsFramework/InteractiveToolsFramework.Build.cs b/Engine/Source/Runtime/InteractiveToolsFramework/InteractiveToolsFramework.Build.cs
--- a/Engine/Source/Runtime/InteractiveToolsFramework/InteractiveToolsFramework.Build.cs
+++ b/Engine/Source/Runtime/InteractiveToolsFramework/InteractiveToolsFramework.Build.cs
@@ -32,5 +32,10 @@ public class InteractiveToolsFramework : ModuleRules
 				// ... add private dependencies that you statically link with here ...	
 			}
 			);
+
+		if (Target.bBuildEditor == true)
+		{
+			PrivateDependencyModuleNames.Add("EditorToolEvents");
+		}
 	}
 }
diff --git a/Engine/Source/Runtime/InteractiveToolsFramework/Private/InteractiveTool.cpp b/Engine/Source/Runtime/InteractiveToolsFramework/Private/InteractiveTool.cpp
--- a/Engine/Source/Runtime/InteractiveToolsFramework/Private/InteractiveTool.cpp
+++ b/Engine/Source/Runtime/InteractiveToolsFramework/Private/InteractiveTool.cpp
@@ -2,6 +2,10 @@
 
 
 #include "InteractiveTool.h"
+
+#if WITH_EDITOR
+#include "EditorToolDelegates.h"
+#endif
 #include "InteractiveToolManager.h"
 #include "UObject/Class.h"
 
@@ -39,6 +43,9 @@ UInteractiveTool::UInteractiveTool()
 
 void UInteractiveTool::Setup()
 {
+#if WITH_EDITOR
+	FEditorToolDelegates::OnEditorToolStarted.Broadcast(GetClass()->GetName());
+#endif
 }
 
 void UInteractiveTool::Shutdown(EToolShutdownType ShutdownType)
