Split the focus summary out of the session summary, so it can be split into it's own event and handled seperately.

- Organize the Focus Summary data so we can make the events easier to parse and filter

#Synced-CL 49063813

diff --git a/Engine/Source/Editor/EditorToolEvents/EditorToolEvents.Build.cs b/Engine/Source/Editor/EditorToolEvents/EditorToolEvents.Build.cs
--- a/Engine/Source/Editor/EditorToolEvents/EditorToolEvents.Build.cs
+++ b/Engine/Source/Editor/EditorToolEvents/EditorToolEvents.Build.cs
@@ -26,8 +26,7 @@ public class EditorToolEvents : ModuleRules
 		PublicDependencyModuleNames.AddRange(
 			new string[]
 			{
-				"Core",
-				"Engine"
+				"Core"
 				// ... add other public dependencies that you statically link with here ...
 			}
 			);
@@ -36,6 +35,9 @@ public class EditorToolEvents : ModuleRules
 		PrivateDependencyModuleNames.AddRange(
 			new string[]
 			{
+				"SlateCore",
+				"Slate",
+				"Json"
 			}
 			);
 		
diff --git a/Engine/Source/Editor/EditorToolEvents/Private/EditorToolDelegates.cpp b/Engine/Source/Editor/EditorToolEvents/Private/EditorToolDelegates.cpp
--- a/Engine/Source/Editor/EditorToolEvents/Private/EditorToolDelegates.cpp
+++ b/Engine/Source/Editor/EditorToolEvents/Private/EditorToolDelegates.cpp
@@ -3,3 +3,4 @@
 #include "EditorToolDelegates.h"
 
 FEditorToolDelegates::FEditorToolStarted	FEditorToolDelegates::OnEditorToolStarted;
+FEditorToolDelegates::FSendEditorFocusSummary	FEditorToolDelegates::OnSendEditorFocusSummary;
diff --git a/Engine/Source/Editor/EditorToolEvents/Private/EditorUsageAnalyticsSummary.cpp b/Engine/Source/Editor/EditorToolEvents/Private/EditorUsageAnalyticsSummary.cpp
new file mode 100644
--- /dev/null
+++ b/Engine/Source/Editor/EditorToolEvents/Private/EditorUsageAnalyticsSummary.cpp
@@ -0,0 +1,205 @@
+// Copyright Epic Games, Inc. All Rights Reserved.
+
+#include "EditorUsageAnalyticsSummary.h"
+
+#include "EditorToolDelegates.h"
+#include "Framework/Application/SlateApplication.h"
+// includes for editor focus summary test
+#include "Framework/Application/SlateUser.h"
+#include "Framework/Docking/SDockingTabStack.h"
+#include "Misc/ConfigCacheIni.h"
+#include "Widgets/Docking/SDockTab.h"
+#include "Widgets/SWindow.h"
+
+// json includes
+#include "Dom/JsonObject.h"
+
+#if WITH_EDITOR
+
+namespace EditorUsageAnalyticsSummaryHelpers
+{
+	bool GetFocusedTabIdentifierIfInFocusPath(TSharedPtr<const FSlateUser> SlateUser, FString& OutTabIdentifier, uint64 &OutWindowId)
+	{
+		if (!SlateUser.IsValid() || !SlateUser->HasValidFocusPath())
+		{
+			return false;
+		}
+
+		const FWeakWidgetPath& WeakFocusPath = SlateUser->GetWeakFocusPath();
+
+		// TODO: just double check we don't end up with nested windows (Main Editor -> Subwindow) if we do we need to do this in the loop instead
+		if (WeakFocusPath.Widgets.Num() > 0 && WeakFocusPath.Widgets[0].IsValid())
+		{
+			OutWindowId = WeakFocusPath.Widgets[0].Pin()->GetId();
+		}
+
+		// intentionally don't check 0 index as that'll be the window
+		for (int i = WeakFocusPath.Widgets.Num() - 1; i > 0; --i)
+		{
+			TWeakPtr<SWidget> WidgetInPathWeakPtr = WeakFocusPath.Widgets[i];
+
+			if (!WidgetInPathWeakPtr.IsValid())
+			{
+				continue;
+			}
+
+			TSharedPtr<SWidget> WidgetInPath = WidgetInPathWeakPtr.Pin();
+
+			// docking tab stack is the part of the docking stuff that shows up in focus paths
+			if (WidgetInPath->GetTypeAsString() == TEXT("SDockingTabStack"))
+			{
+				SDockingTabStack* TabStackWidget = reinterpret_cast<SDockingTabStack*>(WidgetInPath.Get());
+
+				if (TabStackWidget != nullptr)
+				{
+					TSharedPtr<SDockTab> TabWidget = TabStackWidget->GetSelectedTabFromWell();
+
+					if (TabWidget.IsValid())
+					{
+						OutTabIdentifier = TabWidget->GetLayoutIdentifier().ToString();
+
+						return true;
+					}
+
+				}
+			}
+		}
+
+		return false;
+	}
+
+	static const double DefaultFocusedSampleDelay = 3.0;
+}
+
+FEditorUsageAnalyticsSummary& FEditorUsageAnalyticsSummary::Get()
+{
+	static FEditorUsageAnalyticsSummary Summary = FEditorUsageAnalyticsSummary();
+	return Summary;
+}
+
+FEditorUsageAnalyticsSummary::FEditorUsageAnalyticsSummary() : TimeOfLastUpdate(FPlatformTime::Seconds())
+{
+	double Delay = 0.0;
+	if (!GConfig->GetDouble(TEXT("EditorUsageAnalytics"), TEXT("FocusedSampleDelay"), Delay, GEditorIni))
+	{
+		Delay = EditorUsageAnalyticsSummaryHelpers::DefaultFocusedSampleDelay;
+	}
+
+	UpdatePeriod.FromSeconds(Delay);
+}
+
+void FEditorUsageAnalyticsSummary::Start()
+{
+	bActive = true;
+	FSlateApplication::Get().GetLastUserInteractionTimeUpdateEvent().AddRaw(this, &FEditorUsageAnalyticsSummary::OnSlateUserInteraction);
+}
+
+void FEditorUsageAnalyticsSummary::Shutdown()
+{
+	FSlateApplication::Get().GetLastUserInteractionTimeUpdateEvent().RemoveAll(this);
+	bActive = false;
+}
+
+void FEditorUsageAnalyticsSummary::OnSlateUserInteraction(double DeltaTime)
+{
+	double CurrentTimeSeconds = FPlatformTime::Seconds();
+
+	if (CurrentTimeSeconds >= NextUpdateTimeSeconds)
+	{
+		UpdateUserFocusTimes(CurrentTimeSeconds);
+		// update when we want to next update our usage summary statistics. Get updated platform time as updating could have taken some time
+		NextUpdateTimeSeconds = FPlatformTime::Seconds() + UpdatePeriod.GetTotalSeconds();
+	}
+}
+
+void FEditorUsageAnalyticsSummary::UpdateUserFocusTimes(double CurrTimeSecs)
+{
+	// This can get called from a console control handler thread. Getting the SlateApplication is only valid on certain threads.
+	if (!IsInGameThread() && !IsInSlateThread() && !IsInAsyncLoadingThread() && !IsInParallelLoadingThread())
+	{
+		return;
+	}
+
+	// Make sure we have a slate application to query
+	if (!FSlateApplication::IsInitialized())
+	{
+		return;
+	}
+
+	FSlateApplication& SlateApplication = FSlateApplication::Get();
+
+	TSharedPtr<const FSlateUser> CurrentUser = SlateApplication.GetCursorUser();
+
+	double PrevTimeOfLastUpdate = TimeOfLastUpdate.load();
+
+	// we want to ensure only a single thread performs this update
+	if (TimeOfLastUpdate.compare_exchange_strong(PrevTimeOfLastUpdate, CurrTimeSecs))
+	{
+		FString TabIdentifier;
+		uint64 WindowId = 0;
+		if (EditorUsageAnalyticsSummaryHelpers::GetFocusedTabIdentifierIfInFocusPath(CurrentUser, TabIdentifier, WindowId))
+		{
+			float DeltaAsFloat = static_cast<float>(CurrTimeSecs - PrevTimeOfLastUpdate);
+
+			{
+				FScopeLock LockAroundData(&StoreLock);
+
+				FFocusedWindowAnalyticData& WindowData = FocusedWindows.FindOrAdd(WindowId);
+
+				// add a lock around this add if we are worried about multithreaded, which I don't beleive we are at the moment
+
+				if (!WindowData.FocusedTabs.Contains(TabIdentifier))
+				{
+					WindowData.FocusedTabs.Add(TabIdentifier, 0);
+				}
+				else
+				{
+					WindowData.FocusedTabs[TabIdentifier] += DeltaAsFloat;
+				}
+			}
+		}
+	}
+}
+
+TSharedPtr<FJsonObject> FEditorUsageAnalyticsSummary::GetSummaryDataAsJsonObject(bool bReset)
+{
+	TSharedPtr<FJsonObject> SummaryAsJsonObject = MakeShareable(new FJsonObject);
+
+	TArray<TSharedPtr<FJsonValue>> ListOfWindows;
+
+	{
+		FScopeLock LockAroundFocusData(&StoreLock);
+
+		for (const TPair<uint32, FFocusedWindowAnalyticData>& WindowDataPair : FocusedWindows)
+		{
+			TSharedPtr<FJsonObject> WindowEntry = MakeShareable(new FJsonObject);
+
+			WindowEntry->SetNumberField(TEXT("WindowID"), WindowDataPair.Key);
+			TArray<TSharedPtr<FJsonValue>> ListOfTabs;
+			for (const TPair<FString, double> &TabPair : WindowDataPair.Value.FocusedTabs)
+			{
+				TSharedPtr<FJsonObject> TabEntry = MakeShareable(new FJsonObject);
+
+				TabEntry->SetStringField(TEXT("TabName"), TabPair.Key);
+				TabEntry->SetNumberField(TEXT("FocusTime"), TabPair.Value);
+
+				ListOfTabs.Add(TSharedPtr<FJsonValue>(new FJsonValueObject(TabEntry)));
+			}
+
+		
+			WindowEntry->SetArrayField("Tabs", ListOfTabs);
+
+			ListOfWindows.Add(TSharedPtr<FJsonValue>(new FJsonValueObject(WindowEntry)));
+		}
+
+		if (bReset)
+		{
+			FocusedWindows.Reset();
+		}
+	}
+
+	SummaryAsJsonObject->SetArrayField(TEXT("Windows"), ListOfWindows);
+
+	return SummaryAsJsonObject;
+}
+#endif// WITH_EDITOR
diff --git a/Engine/Source/Editor/EditorToolEvents/Public/EditorToolDelegates.h b/Engine/Source/Editor/EditorToolEvents/Public/EditorToolDelegates.h
--- a/Engine/Source/Editor/EditorToolEvents/Public/EditorToolDelegates.h
+++ b/Engine/Source/Editor/EditorToolEvents/Public/EditorToolDelegates.h
@@ -5,7 +5,10 @@
  */
 
 #pragma once
+
 #include "Delegates/Delegate.h"
+#include "Dom/JsonObject.h"
+#include "Templates/SharedPointer.h"
 
 struct FEditorToolDelegates
 {
@@ -15,5 +18,7 @@ public:
 	DECLARE_MULTICAST_DELEGATE_OneParam(FEditorToolStarted, const FString &/*class name*/);
 	static EDITORTOOLEVENTS_API FEditorToolStarted OnEditorToolStarted;
 	
-	/*More delegates can be added, starting minimal so only events we use are added*/
+	/**Broadcasted when the editor is closed and sends a summary of widget focus data */
+	DECLARE_MULTICAST_DELEGATE_OneParam(FSendEditorFocusSummary, TSharedPtr<FJsonObject>/*json summary of focus data*/);
+	static EDITORTOOLEVENTS_API FSendEditorFocusSummary OnSendEditorFocusSummary;
 };
diff --git a/Engine/Source/Editor/EditorToolEvents/Public/EditorUsageAnalyticsSummary.h b/Engine/Source/Editor/EditorToolEvents/Public/EditorUsageAnalyticsSummary.h
new file mode 100644
--- /dev/null
+++ b/Engine/Source/Editor/EditorToolEvents/Public/EditorUsageAnalyticsSummary.h
@@ -0,0 +1,64 @@
+// Copyright Epic Games, Inc. All Rights Reserved.
+
+#pragma once
+
+#include "CoreMinimal.h"
+#include "Misc/Optional.h"
+#include "Misc/Timespan.h"
+
+#if WITH_EDITOR
+
+struct FFocusedWindowAnalyticData
+{
+	// Name, time in seconds focused
+	TMap<FString, double> FocusedTabs;
+	uint32 WindowId;
+};
+
+class FJsonObject;
+
+/*
+*	Collects editor events we wish to send as a summary that are not critical enough to be in our summary store.
+*/
+class FEditorUsageAnalyticsSummary : FNoncopyable
+{
+public:
+	EDITORTOOLEVENTS_API static FEditorUsageAnalyticsSummary& Get();
+
+	EDITORTOOLEVENTS_API FEditorUsageAnalyticsSummary();
+
+	EDITORTOOLEVENTS_API void Start();
+
+	EDITORTOOLEVENTS_API void Shutdown();
+
+	// Get's the current usage summary. If bReset is true, clears out current summary. (useful when collecting summaries per project instead of session)
+	EDITORTOOLEVENTS_API TSharedPtr<FJsonObject> GetSummaryDataAsJsonObject(bool bReset = false);
+
+	// Returns true if start has been called. Makes it so we can still check the summary even if the CVAR to turn it off is true or otherwise disabled.
+	bool IsActive() { return bActive; }
+
+	virtual ~FEditorUsageAnalyticsSummary() = default;
+
+private:
+
+	void OnSlateUserInteraction(double DeltaTime);
+	void UpdateUserFocusTimes(double CurrTimeSecs);
+
+	TMap<uint32, FFocusedWindowAnalyticData> FocusedWindows;
+
+	// Time since we last updated our focus stats
+	std::atomic<double> TimeOfLastUpdate;
+
+	//Slate interaction and shutting down can come from different threads so we need a log for FocusedWindow data
+	mutable FTransactionallySafeCriticalSection StoreLock;
+
+	//Period we want to wait between focus time updates (some events might lead us to update sooner)
+	FTimespan UpdatePeriod;
+
+	//The next time the session properties should be persisted to disk.
+	double NextUpdateTimeSeconds = 0.0;
+
+	bool bActive = false;
+};
+
+#endif
\ No newline at end of file
diff --git a/Engine/Source/Runtime/Engine/Engine.Build.cs b/Engine/Source/Runtime/Engine/Engine.Build.cs
--- a/Engine/Source/Runtime/Engine/Engine.Build.cs
+++ b/Engine/Source/Runtime/Engine/Engine.Build.cs
@@ -72,7 +72,7 @@ public class Engine : ModuleRules
 			PrivateDependencyModuleNames.AddRange(
 				new string[]
 				{
-					"EditorAnalyticsSession",
+					"EditorAnalyticsSession"
 				}
 			);
 		}
@@ -219,6 +219,9 @@ public class Engine : ModuleRules
 					"Zen"
 				}
 			);
+
+			// editor telemetry events
+			PrivateDependencyModuleNames.Add("EditorToolEvents");
 		}
 
 		PrivateIncludePathModuleNames.AddRange(
diff --git a/Engine/Source/Runtime/Engine/Private/Analytics/EditorAnalyticsSessionSummary.cpp b/Engine/Source/Runtime/Engine/Private/Analytics/EditorAnalyticsSessionSummary.cpp
--- a/Engine/Source/Runtime/Engine/Private/Analytics/EditorAnalyticsSessionSummary.cpp
+++ b/Engine/Source/Runtime/Engine/Private/Analytics/EditorAnalyticsSessionSummary.cpp
@@ -13,6 +13,8 @@
 #include "Framework/Docking/SDockingTabStack.h"
 #include "Widgets/Docking/SDockTab.h"
 
+#include "EditorUsageAnalyticsSummary.h"
+
 #include "Interfaces/IAnalyticsPropertyStore.h"
 #include "Kismet2/DebuggerCommands.h"
 #include "ProfilingDebugging/StallDetector.h"
@@ -54,11 +56,11 @@ namespace EditorAnalyticsCVars
 		TEXT("Disables collecting focused widget times for analytics summary."));
 }
 
+
 FEditorAnalyticsSessionSummary::FEditorAnalyticsSessionSummary(TSharedPtr<IAnalyticsPropertyStore> Store, uint32 MonitorProcessId)
 	: FEngineAnalyticsSessionSummary(MoveTemp(Store), MonitorProcessId)
 	, LastUserActivityTimeSecs(FPlatformTime::Seconds())
 	, AccountedUserIdleSecs(0)
-	, TimeOfLastUpdate(FPlatformTime::Seconds())
 {
 	EditorAnalyticsProperties::UserInteractionCount.Set(GetStore(), 0);
 	EditorAnalyticsProperties::Idle1Min.Set(GetStore(), 0);
@@ -98,10 +100,21 @@ FEditorAnalyticsSessionSummary::FEditorAnalyticsSessionSummary(TSharedPtr<IAnaly
 	FSlateApplication::Get().GetOnModalLoopTickEvent().AddRaw(this, &FEditorAnalyticsSessionSummary::Tick);
 	FSlateApplication::Get().GetLastUserInteractionTimeUpdateEvent().AddRaw(this, &FEditorAnalyticsSessionSummary::OnSlateUserInteraction);
 	UPackage::PackageDirtyStateChangedEvent.AddRaw(this, &FEditorAnalyticsSessionSummary::OnDirtyPackageStateChanged);
+
+	if (!EditorAnalyticsCVars::CVarDisableFocusTimes.GetValueOnAnyThread())
+	{
+		// we want to make sure an instance is created so the usage summary starts listening to events
+		FEditorUsageAnalyticsSummary::Get().Start();
+	}
 }
 
 void FEditorAnalyticsSessionSummary::ShutdownInternal()
 {
+	if (!EditorAnalyticsCVars::CVarDisableFocusTimes.GetValueOnAnyThread())
+	{
+		FEditorUsageAnalyticsSummary::Get().Shutdown();
+	}
+
 	FEditorDelegates::PreBeginPIE.RemoveAll(this);
 	FEditorDelegates::EndPIE.RemoveAll(this);
 	FSlateApplication::Get().GetOnModalLoopTickEvent().RemoveAll(this);
@@ -145,106 +158,9 @@ bool FEditorAnalyticsSessionSummary::UpdateSessionProgressInternal(bool bCrashin
 	EditorAnalyticsProperties::ShaderDDCHits.Set(GetStore(), GShaderCompilerStats->GetDDCHits());
 	EditorAnalyticsProperties::TimeShaderCompilationWasActive.Set(GetStore(), GShaderCompilerStats->GetTimeShaderCompilationWasActive());
 
-	if (!EditorAnalyticsCVars::CVarDisableFocusTimes.GetValueOnAnyThread())
-	{
-		UpdateUserFocusTimes(FPlatformTime::Seconds());
-	}
-
 	return bShouldPersist;
 }
 
-namespace EditorAnalyticsSessionSummaryHelpers
-{
-	bool GetFocusedTabIdentifierIfInFocusPath(TSharedPtr<const FSlateUser> SlateUser, FString& OutTabIdentifier)
-	{
-		if (!SlateUser.IsValid() || !SlateUser->HasValidFocusPath())
-		{
-			return false;
-		}
-
-		const FWeakWidgetPath& WeakFocusPath = SlateUser->GetWeakFocusPath();
-
-		// intentionally don't check 0 index as that'll be the window
-		for (int i = WeakFocusPath.Widgets.Num() - 1; i > 0; --i)
-		{
-			TWeakPtr<SWidget> WidgetInPathWeakPtr = WeakFocusPath.Widgets[i];
-
-			if (!WidgetInPathWeakPtr.IsValid())
-			{
-				continue;
-			}
-
-			TSharedPtr<SWidget> WidgetInPath = WidgetInPathWeakPtr.Pin();
-
-			// docking tab stack is the part of the docking stuff that shows up in focus paths
-			if (WidgetInPath->GetTypeAsString() == TEXT("SDockingTabStack"))
-			{
-				SDockingTabStack* TabStackWidget = reinterpret_cast<SDockingTabStack*>(WidgetInPath.Get());
-
-				if (TabStackWidget != nullptr)
-				{
-					TSharedPtr<SDockTab> TabWidget = TabStackWidget->GetSelectedTabFromWell();
-
-					if (TabWidget.IsValid())
-					{
-						OutTabIdentifier = TabWidget->GetLayoutIdentifier().ToString();
-
-						return true;
-					}
-
-				}
-			}
-		}
-
-		return false;
-	}
-}
-
-void FEditorAnalyticsSessionSummary::UpdateUserFocusTimes(double CurrTimeSecs)
-{
-	// This can get called from a console control handler thread. Getting the SlateApplication is only valid on certain threads.
-	if (!IsInGameThread() && !IsInSlateThread() && !IsInAsyncLoadingThread() && !IsInParallelLoadingThread())
-	{
-		return;
-	}
-
-	// Make sure we have a slate application to query
-	if (!FSlateApplication::IsInitialized())
-	{
-		return;
-	}
-
-	FSlateApplication& SlateApplication = FSlateApplication::Get();
-
-	TSharedPtr<const FSlateUser> CurrentUser = SlateApplication.GetCursorUser();
-
-	double PrevTimeOfLastUpdate = TimeOfLastUpdate.load();
-
-	// similar to our other timer updates, we want to ensure only a single thread performs this update
-	if (TimeOfLastUpdate.compare_exchange_strong(PrevTimeOfLastUpdate, CurrTimeSecs))
-	{
-		FString TabIdentifier;
-		if (EditorAnalyticsSessionSummaryHelpers::GetFocusedTabIdentifierIfInFocusPath(CurrentUser, TabIdentifier))
-		{
-			TabIdentifier = FString(TEXT("FocusTime")).Append(TabIdentifier);
-			float DeltaAsFloat = static_cast<float>(CurrTimeSecs - PrevTimeOfLastUpdate);
-
-			if (!GetStore()->Contains(TabIdentifier))
-			{
-				GetStore()->Set(TabIdentifier, DeltaAsFloat);
-			}
-			else
-			{
-				GetStore()->Update(TabIdentifier, [DeltaAsFloat](float& InOutValue)
-					{
-						InOutValue += DeltaAsFloat;
-						return true;
-					});
-			}
-		}
-	}
-}
-
 bool FEditorAnalyticsSessionSummary::UpdateUserIdleTime(double CurrTimeSecs, bool bReset)
 {
 	bool bSessionUpdated = false;
diff --git a/Engine/Source/Runtime/Engine/Private/Analytics/EditorAnalyticsSessionSummary.h b/Engine/Source/Runtime/Engine/Private/Analytics/EditorAnalyticsSessionSummary.h
--- a/Engine/Source/Runtime/Engine/Private/Analytics/EditorAnalyticsSessionSummary.h
+++ b/Engine/Source/Runtime/Engine/Private/Analytics/EditorAnalyticsSessionSummary.h
@@ -26,7 +26,6 @@ protected:
 
 private:
 	bool UpdateUserIdleTime(double CurrTimeSecs, bool bReset);
-	void UpdateUserFocusTimes(double CurrTimeSecs);
 	void OnSlateUserInteraction(double CurrSlateInteractionTime);
 	void OnEnterPIE(const bool /*bIsSimulating*/);
 	void OnExitPIE(const bool /*bIsSimulating*/);
@@ -38,9 +37,6 @@ private:
 
 	/** The number of idle seconds in the current idle sequence that were accounted (saved in the session) for the user idle counters. */
 	std::atomic<double> AccountedUserIdleSecs;
-
-	/** Time since we last updated our focus stats */
-	std::atomic<double> TimeOfLastUpdate;
 };
 
 #endif // WITH_EDITOR
