Fix error in Focus Summary when we don't have widget unique identifiers enabled.

#Synced-CL 49063813

diff --git a/Engine/Source/Editor/EditorToolEvents/Private/EditorUsageAnalyticsSummary.cpp b/Engine/Source/Editor/EditorToolEvents/Private/EditorUsageAnalyticsSummary.cpp
--- a/Engine/Source/Editor/EditorToolEvents/Private/EditorUsageAnalyticsSummary.cpp
+++ b/Engine/Source/Editor/EditorToolEvents/Private/EditorUsageAnalyticsSummary.cpp
@@ -27,11 +27,15 @@ namespace EditorUsageAnalyticsSummaryHelpers
 
 		const FWeakWidgetPath& WeakFocusPath = SlateUser->GetWeakFocusPath();
 
-		// TODO: just double check we don't end up with nested windows (Main Editor -> Subwindow) if we do we need to do this in the loop instead
+		#if UE_SLATE_WITH_WIDGET_UNIQUE_IDENTIFIER
 		if (WeakFocusPath.Widgets.Num() > 0 && WeakFocusPath.Widgets[0].IsValid())
 		{
 			OutWindowId = WeakFocusPath.Widgets[0].Pin()->GetId();
 		}
+		#else
+		// if we can't use the widget ID we have to lump all of our focus data as if they are in a single window
+		OutWindowId = 0;
+		#endif
 
 		// intentionally don't check 0 index as that'll be the window
 		for (int i = WeakFocusPath.Widgets.Num() - 1; i > 0; --i)
