Fix missing ToString from window title in focus summary.

#Synced-CL 49631393

diff --git a/Engine/Source/Editor/EditorToolEvents/Private/EditorUsageAnalyticsSummary.cpp b/Engine/Source/Editor/EditorToolEvents/Private/EditorUsageAnalyticsSummary.cpp
--- a/Engine/Source/Editor/EditorToolEvents/Private/EditorUsageAnalyticsSummary.cpp
+++ b/Engine/Source/Editor/EditorToolEvents/Private/EditorUsageAnalyticsSummary.cpp
@@ -45,7 +45,7 @@ namespace EditorUsageAnalyticsSummaryHelpers
 
 				if (WindowWidget != nullptr)
 				{
-					OutWindowId = WindowWidget->GetTitle();
+					OutWindowId = WindowWidget->GetTitle().ToString();
 				}
 				else
 				{
